Electra VPx decoder: added VP9 support for 10 bit HDR output

#rnx

diff --git a/Engine/Plugins/Media/VPxDecoderElectra/Source/VPxDecoderElectra/Private/VPxDecoder/ElectraMediaVPxDecoder.cpp b/Engine/Plugins/Media/VPxDecoderElectra/Source/VPxDecoderElectra/Private/VPxDecoder/ElectraMediaVPxDecoder.cpp
--- a/Engine/Plugins/Media/VPxDecoderElectra/Source/VPxDecoderElectra/Private/VPxDecoder/ElectraMediaVPxDecoder.cpp
+++ b/Engine/Plugins/Media/VPxDecoderElectra/Source/VPxDecoderElectra/Private/VPxDecoder/ElectraMediaVPxDecoder.cpp
@@ -43,6 +43,7 @@ public:
 
 // This should really be handled differently...
 static bool ConvertDecodedImageToNV12(TSharedPtr<TArray<uint8>, ESPMode::ThreadSafe>& OutNV12Buffer, const vpx_image_t* InDecodedImage);
+static bool ConvertDecodedImageToP010(TSharedPtr<TArray<uint8>, ESPMode::ThreadSafe>& OutP010Buffer, const vpx_image_t* InDecodedImage);
 
 
 class FVideoDecoderOutputVPxElectra : public IElectraDecoderVideoOutput, public IElectraDecoderVideoOutputImageBuffers
@@ -797,6 +798,21 @@ FVideoDecoderVPxElectra::EConvertResult FVideoDecoderVPxElectra::ConvertDecoderO
 		NewOutput->DecodedWidth = NewOutput->Width;
 		NewOutput->DecodedHeight = NewOutput->Height * 3 / 2;
 	}
+	else if (InDecodedImage->fmt == VPX_IMG_FMT_I42016)
+	{
+		check(NewOutput->NumBits == 10);
+		if (!ConvertDecodedImageToP010(NewOutput->ColorBuffer, InDecodedImage))
+		{
+			PostError(0, FString::Printf(TEXT("Failed to convert decoded image")), ERRCODE_INTERNAL_FAILED_TO_CONVERT_OUTPUT_SAMPLE);
+			return EConvertResult::Failure;
+		}
+		NewOutput->NumBuffers = 1;
+		NewOutput->ColorBufferFormat = EElectraDecoderPlatformPixelFormat::P010;
+		NewOutput->ColorBufferEncoding = EElectraDecoderPlatformPixelEncoding::Native;
+		NewOutput->ColorPitch = NewOutput->Width * 2;
+		NewOutput->DecodedWidth = NewOutput->Width;
+		NewOutput->DecodedHeight = NewOutput->Height * 3 / 2;
+	}
 	else
 	{
 		PostError(0, FString::Printf(TEXT("Unsupported decoded image format (%d)"), InDecodedImage->fmt), ERRCODE_INTERNAL_FAILED_TO_CONVERT_OUTPUT_SAMPLE);
@@ -830,7 +846,7 @@ static bool ConvertDecodedImageToNV12(TSharedPtr<TArray<uint8>, ESPMode::ThreadS
 	const int32 aw = Align(w, 2);
 	const int32 ah = Align(h, 2);
 
-	int32 AllocSize = aw * ah * 3 /2;
+	int32 AllocSize = aw * ah * 3/2;
 
 	OutNV12Buffer = MakeShared<TArray<uint8>, ESPMode::ThreadSafe>();
 	OutNV12Buffer->AddUninitialized(AllocSize);
@@ -867,3 +883,51 @@ static bool ConvertDecodedImageToNV12(TSharedPtr<TArray<uint8>, ESPMode::ThreadS
 	}
 	return true;
 }
+
+static bool ConvertDecodedImageToP010(TSharedPtr<TArray<uint8>, ESPMode::ThreadSafe>& OutP010Buffer, const vpx_image_t* InDecodedImage)
+{
+	const int32 w = InDecodedImage->d_w;
+	const int32 h = InDecodedImage->d_h;
+	const int32 aw = Align(w, 2);
+	const int32 ah = Align(h, 2);
+
+	int32 AllocSize = aw * ah * 3;
+
+	OutP010Buffer = MakeShared<TArray<uint8>, ESPMode::ThreadSafe>();
+	OutP010Buffer->AddUninitialized(AllocSize);
+
+	uint16* DstY = (uint16*)OutP010Buffer->GetData();
+	uint16* DstUV = DstY + aw * ah;
+	const uint16* SrcY = (const uint16*) InDecodedImage->planes[0];
+	const uint16* SrcU = (const uint16*) InDecodedImage->planes[1];
+	const uint16* SrcV = (const uint16*) InDecodedImage->planes[2];
+	const int32 PitchY = InDecodedImage->stride[0] / 2;
+	const int32 PitchU = InDecodedImage->stride[1] / 2;
+	const int32 PitchV = InDecodedImage->stride[2] / 2;
+	if (!SrcY || !SrcU || !SrcV)
+	{
+		return false;
+	}
+	for(int32 y=0; y<h; ++y)
+	{
+		for(int32 x=0; x<w; ++x)
+		{
+			DstY[x] = SrcY[x] << 6;
+		}
+		DstY += aw;
+		SrcY += PitchY;
+	}
+	int32 padUV = (aw - w) * 2;
+	for(int32 v=0; v<h/2; ++v)
+	{
+		for(int32 u=0; u<w/2; ++u)
+		{
+			*DstUV++ = SrcU[u] << 6;
+			*DstUV++ = SrcV[u] << 6;
+		}
+		SrcU += PitchU;
+		SrcV += PitchV;
+		DstUV += padUV;
+	}
+	return true;
+}
