ElectraCodecs: Fixed passing of HDR metadata with VP9 in webm/mkv container; added VP9 alpha channel support when using libvpx for decoding

#rnx

diff --git a/Engine/Plugins/Media/ElectraCodecs/Source/ElectraDecoders/Private/Utils/Google/ElectraBitstreamProcessorVPx.cpp b/Engine/Plugins/Media/ElectraCodecs/Source/ElectraDecoders/Private/Utils/Google/ElectraBitstreamProcessorVPx.cpp
--- a/Engine/Plugins/Media/ElectraCodecs/Source/ElectraDecoders/Private/Utils/Google/ElectraBitstreamProcessorVPx.cpp
+++ b/Engine/Plugins/Media/ElectraCodecs/Source/ElectraDecoders/Private/Utils/Google/ElectraBitstreamProcessorVPx.cpp
@@ -35,7 +35,7 @@ public:
 		return LastErrorMessage;
 	}
 
-	bool ProcessInputForDecoding(FElectraDecoderInputAccessUnit& InOutAccessUnit)
+	bool ProcessInputForDecoding(FElectraDecoderInputAccessUnit& InOutAccessUnit, const TMap<FString, FVariant>& InAccessUnitSidebandData)
 	{
 		InOutAccessUnit.Flags &= ~EElectraDecoderFlags::IsDiscardable;
 		if (Version == EVersion::VP8)
@@ -61,22 +61,27 @@ public:
 		else if (Version == EVersion::VP9)
 		{
 			ElectraDecodersUtil::VPxVideo::FVP9UncompressedHeader Header;
-			if (ElectraDecodersUtil::VPxVideo::ParseVP9UncompressedHeader(Header, InOutAccessUnit.Data, InOutAccessUnit.DataSize))
+			if (!ElectraDecodersUtil::VPxVideo::ParseVP9UncompressedHeader(Header, InOutAccessUnit.Data, InOutAccessUnit.DataSize))
 			{
-				if (Header.IsKeyframe())
-				{
-					InOutAccessUnit.Flags |= EElectraDecoderFlags::IsSyncSample;
-				}
-				else
-				{
-					InOutAccessUnit.Flags &= ~EElectraDecoderFlags::IsSyncSample;
-				}
-				return true;
+				LastErrorMessage = FString::Printf(TEXT("Failed to parse VP9 header"));
+				return false;
+			}
+			if (Header.IsKeyframe())
+			{
+				InOutAccessUnit.Flags |= EElectraDecoderFlags::IsSyncSample;
 			}
 			else
 			{
-				LastErrorMessage = FString::Printf(TEXT("Failed to parse VP9 header"));
+				InOutAccessUnit.Flags &= ~EElectraDecoderFlags::IsSyncSample;
+			}
+#if 0
+			// Per-frame HDR info?
+			TArray<uint8> ITUT35SidebandData = ElectraDecodersUtil::GetVariantValueUInt8Array(InAccessUnitSidebandData, TEXT("itu-t-35"));
+			if (ITUT35SidebandData.Num())
+			{
 			}
+#endif
+			return true;
 		}
 		return false;
 	}
@@ -243,7 +248,7 @@ IElectraDecoderBitstreamProcessor::EProcessResult FElectraDecoderBitstreamProces
 	}
 
 	InOutAccessUnit.Flags |= EElectraDecoderFlags::InputIsProcessed;
-	if (!Impl->ProcessInputForDecoding(InOutAccessUnit))
+	if (!Impl->ProcessInputForDecoding(InOutAccessUnit, InAccessUnitSidebandData))
 	{
 		return EProcessResult::Error;
 	}
diff --git a/Engine/Plugins/Media/ElectraCodecs/Source/ElectraDecoders/Private/Utils/MPEG/ElectraUtilsMPEGVideo.cpp b/Engine/Plugins/Media/ElectraCodecs/Source/ElectraDecoders/Private/Utils/MPEG/ElectraUtilsMPEGVideo.cpp
--- a/Engine/Plugins/Media/ElectraCodecs/Source/ElectraDecoders/Private/Utils/MPEG/ElectraUtilsMPEGVideo.cpp
+++ b/Engine/Plugins/Media/ElectraCodecs/Source/ElectraDecoders/Private/Utils/MPEG/ElectraUtilsMPEGVideo.cpp
@@ -192,6 +192,24 @@ namespace ElectraDecodersUtil
 			}
 			return true;
 		}
+		TArray<uint8> ELECTRADECODERS_API BuildMDCVBox(const FSEImastering_display_colour_volume& InFromDisplayColorVolume)
+		{
+			FElectraBitstreamWriter bw;
+			// The order in the MDCV box is G,B,R. Since we assume that in this structure the order is RGB we swap it.
+			bw.PutBits((uint16)InFromDisplayColorVolume.display_primaries_x[1], 16);
+			bw.PutBits((uint16)InFromDisplayColorVolume.display_primaries_y[1], 16);
+			bw.PutBits((uint16)InFromDisplayColorVolume.display_primaries_x[2], 16);
+			bw.PutBits((uint16)InFromDisplayColorVolume.display_primaries_y[2], 16);
+			bw.PutBits((uint16)InFromDisplayColorVolume.display_primaries_x[0], 16);
+			bw.PutBits((uint16)InFromDisplayColorVolume.display_primaries_y[0], 16);
+			bw.PutBits((uint16)InFromDisplayColorVolume.white_point_x, 16);
+			bw.PutBits((uint16)InFromDisplayColorVolume.white_point_y, 16);
+			bw.PutBits((uint32)InFromDisplayColorVolume.max_display_mastering_luminance, 32);
+			bw.PutBits((uint32)InFromDisplayColorVolume.min_display_mastering_luminance, 32);
+			TArray<uint8> Box;
+			bw.GetArray(Box);
+			return Box;
+		}
 
 
 		bool ParseFromSEIMessage(FSEIcontent_light_level_info& OutCLLI, const FSEIMessage& InMessage)
@@ -235,6 +253,16 @@ namespace ElectraDecodersUtil
 			}
 			return true;
 		}
+		TArray<uint8> BuildCLLIBox(const FSEIcontent_light_level_info& InFromContentLightLevel)
+		{
+			FElectraBitstreamWriter bw;
+			bw.PutBits((uint16)InFromContentLightLevel.max_content_light_level, 16);
+			bw.PutBits((uint16)InFromContentLightLevel.max_pic_average_light_level, 16);
+			TArray<uint8> Box;
+			bw.GetArray(Box);
+			return Box;
+		}
+
 
 
 
@@ -280,6 +308,20 @@ namespace ElectraDecodersUtil
 			return false;
 		}
 
+		TArray<uint8> ELECTRADECODERS_API BuildCOLRBox(const FCommonColorimetry& InFromColorimetry)
+		{
+			FElectraBitstreamWriter bw;
+			bw.PutBits(ElectraDecodersUtil::MakeMP4Atom('n','c','l','x'), 32);
+			bw.PutBits((uint16)InFromColorimetry.colour_primaries, 16);
+			bw.PutBits((uint16)InFromColorimetry.transfer_characteristics, 16);
+			bw.PutBits((uint16)InFromColorimetry.matrix_coeffs, 16);
+			bw.PutBits(InFromColorimetry.video_full_range_flag, 1);
+			bw.PutBits(0, 7);
+			TArray<uint8> Box;
+			bw.GetArray(Box);
+			return Box;
+		}
+
 
 	} // namespace MPEG
 } // namespace ElectraDecodersUtil
diff --git a/Engine/Plugins/Media/ElectraCodecs/Source/ElectraDecoders/Public/Utils/MPEG/ElectraUtilsMPEGVideo.h b/Engine/Plugins/Media/ElectraCodecs/Source/ElectraDecoders/Public/Utils/MPEG/ElectraUtilsMPEGVideo.h
--- a/Engine/Plugins/Media/ElectraCodecs/Source/ElectraDecoders/Public/Utils/MPEG/ElectraUtilsMPEGVideo.h
+++ b/Engine/Plugins/Media/ElectraCodecs/Source/ElectraDecoders/Public/Utils/MPEG/ElectraUtilsMPEGVideo.h
@@ -50,6 +50,7 @@ namespace ElectraDecodersUtil
 		};
 		bool ELECTRADECODERS_API ParseFromSEIMessage(FSEImastering_display_colour_volume& OutMDCV, const FSEIMessage& InMessage);
 		bool ELECTRADECODERS_API ParseFromMDCVBox(FSEImastering_display_colour_volume& OutMDCV, const TConstArrayView<uint8>& InMDCVBox);
+		TArray<uint8> ELECTRADECODERS_API BuildMDCVBox(const FSEImastering_display_colour_volume& InFromDisplayColorVolume);
 
 		struct FSEIcontent_light_level_info
 		{
@@ -59,6 +60,7 @@ namespace ElectraDecodersUtil
 		bool ELECTRADECODERS_API ParseFromSEIMessage(FSEIcontent_light_level_info& OutCLLI, const FSEIMessage& InMessage);
 		bool ELECTRADECODERS_API ParseFromCOLLBox(FSEIcontent_light_level_info& OutCLLI, const TConstArrayView<uint8>& InCOLLBox);
 		bool ELECTRADECODERS_API ParseFromCLLIBox(FSEIcontent_light_level_info& OutCLLI, const TConstArrayView<uint8>& InCLLIBox);
+		TArray<uint8> ELECTRADECODERS_API BuildCLLIBox(const FSEIcontent_light_level_info& InFromContentLightLevel);
 
 		struct FSEIalternative_transfer_characteristics
 		{
@@ -98,6 +100,8 @@ namespace ElectraDecodersUtil
 			uint8 video_format = 5;
 		};
 		bool ELECTRADECODERS_API ParseFromCOLRBox(FCommonColorimetry& OutColorimetry, const TConstArrayView<uint8>& InCOLRBox);
+		TArray<uint8> ELECTRADECODERS_API BuildCOLRBox(const FCommonColorimetry& InFromColorimetry);
+
 
 		struct FCommonPictureTiming
 		{
diff --git a/Engine/Plugins/Media/ElectraPlayer/Source/ElectraPlayerRuntime/Private/Runtime/Decoder/VideoDecoder.cpp b/Engine/Plugins/Media/ElectraPlayer/Source/ElectraPlayerRuntime/Private/Runtime/Decoder/VideoDecoder.cpp
--- a/Engine/Plugins/Media/ElectraPlayer/Source/ElectraPlayerRuntime/Private/Runtime/Decoder/VideoDecoder.cpp
+++ b/Engine/Plugins/Media/ElectraPlayer/Source/ElectraPlayerRuntime/Private/Runtime/Decoder/VideoDecoder.cpp
@@ -715,6 +715,14 @@ FVideoDecoderImpl::EAUChangeFlags FVideoDecoderImpl::PrepareAU(TSharedPtrTS<FDec
 		{
 			AU->CSDOptions = CurrentActiveCSD;
 		}
+		// Pass dynamic sideband data
+		if (AU->AccessUnit->DynamicSidebandData.IsValid())
+		{
+			for(auto& dynsbIt : *AU->AccessUnit->DynamicSidebandData)
+			{
+				AU->CSDOptions.Emplace(dynsbIt.Key.ToString(), FVariant(dynsbIt.Value));
+			}
+		}
 
 		// Set the timestamps in the decoder input description.
 		AU->DecAU.DTS = AU->AccessUnit->DTS.GetAsTimespan();
diff --git a/Engine/Plugins/Media/ElectraPlayer/Source/ElectraPlayerRuntime/Private/Runtime/Demuxer/ParserISO14496-12.cpp b/Engine/Plugins/Media/ElectraPlayer/Source/ElectraPlayerRuntime/Private/Runtime/Demuxer/ParserISO14496-12.cpp
--- a/Engine/Plugins/Media/ElectraPlayer/Source/ElectraPlayerRuntime/Private/Runtime/Demuxer/ParserISO14496-12.cpp
+++ b/Engine/Plugins/Media/ElectraPlayer/Source/ElectraPlayerRuntime/Private/Runtime/Demuxer/ParserISO14496-12.cpp
@@ -7385,120 +7385,6 @@ namespace Electra
 		return ParseError;
 	}
 
-	static UEMediaError ParseCOLRBox(FStreamCodecInformation& OutInfo, const FMP4Box* InCOLR)
-	{
-		UEMediaError Error = UEMEDIA_ERROR_OK;
-		const FMP4BoxIgnored* Box = static_cast<const FMP4BoxIgnored*>(InCOLR);
-		if (Box && Box->GetFillerData())
-		{
-			TArray<uint8> BoxData((const uint8*)Box->GetFillerData()->Data, Box->GetFillerData()->Size);
-			FDataBufferReader BufferReader(BoxData);
-			FMP4BoxReader Reader(&BufferReader, nullptr);
-			uint32 Type;
-			uint16 colour_primaries, transfer_characteristics, matrix_coeffs;
-			uint8 video_full_range_flag = 0;
-			RETURN_IF_ERROR(Reader.Read(Type));
-			if (Type == Utils::Make4CC('n','c','l','x') || Type == Utils::Make4CC('n', 'c', 'l', 'c'))
-			{
-				RETURN_IF_ERROR(Reader.Read(colour_primaries));
-				RETURN_IF_ERROR(Reader.Read(transfer_characteristics));
-				RETURN_IF_ERROR(Reader.Read(matrix_coeffs));
-				if (Type == Utils::Make4CC('n', 'c', 'l', 'x'))
-				{
-					RETURN_IF_ERROR(Reader.Read(video_full_range_flag));
-				}
-				OutInfo.GetCodecVideoColorInfo().ColourPrimaries = colour_primaries;
-				OutInfo.GetCodecVideoColorInfo().TransferCharacteristics = transfer_characteristics;
-				OutInfo.GetCodecVideoColorInfo().MatrixCoefficients = matrix_coeffs;
-				OutInfo.GetCodecVideoColorInfo().VideoFullRangeFlag = video_full_range_flag;
-			}
-		}
-		return Error;
-	}
-
-
-	static UEMediaError ParseCOLLBox(FStreamCodecInformation& OutInfo, const FMP4Box* InCOLL, bool bIsCOLL)
-	{
-		UEMediaError Error = UEMEDIA_ERROR_OK;
-		const FMP4BoxIgnored* Box = static_cast<const FMP4BoxIgnored*>(InCOLL);
-		if (Box && Box->GetFillerData())
-		{
-			TArray<uint8> BoxData((const uint8*)Box->GetFillerData()->Data, Box->GetFillerData()->Size);
-			FDataBufferReader BufferReader(BoxData);
-			FMP4BoxReader Reader(&BufferReader, nullptr);
-			if (bIsCOLL)
-			{
-				uint8 Version;
-				RETURN_IF_ERROR(Reader.Read(Version));
-				// A version 0 `coll` box is the same as a `clli` box, which we can handle.
-				if (Version)
-				{
-					return Error;
-				}
-				// Skip the remaining 3 bytes of the full box.
-				Reader.Read(Version);
-				Reader.Read(Version);
-				Reader.Read(Version);
-			}
-			// Read the `clli` box.
-			uint16 max_content_light_level;			// MaxCLL
-			uint16 max_pic_average_light_level;		// MaxFALL
-			RETURN_IF_ERROR(Reader.Read(max_content_light_level));
-			RETURN_IF_ERROR(Reader.Read(max_pic_average_light_level));
-			OutInfo.GetCodecVideoColorInfo().CLLI = { max_content_light_level, max_pic_average_light_level };
-		}
-		return Error;
-	}
-
-	static UEMediaError ParseMDCVBox(FStreamCodecInformation& OutInfo, const FMP4Box* InMDCV)
-	{
-		UEMediaError Error = UEMEDIA_ERROR_OK;
-		const FMP4BoxIgnored* Box = static_cast<const FMP4BoxIgnored*>(InMDCV);
-		if (Box && Box->GetFillerData())
-		{
-			TArray<uint8> BoxData((const uint8*)Box->GetFillerData()->Data, Box->GetFillerData()->Size);
-			FDataBufferReader BufferReader(BoxData);
-			FMP4BoxReader Reader(&BufferReader, nullptr);
-
-			struct mastering_display_colour_volume
-			{
-				uint16 display_primaries_x[3]{ 0 };
-				uint16 display_primaries_y[3]{ 0 };
-				uint16 white_point_x = 0;
-				uint16 white_point_y = 0;
-				uint32 max_display_mastering_luminance = 0;
-				uint32 min_display_mastering_luminance = 0;
-			};
-			mastering_display_colour_volume boxData;
-			for(int32 i=0; i<3; ++i)
-			{
-				RETURN_IF_ERROR(Reader.Read(boxData.display_primaries_x[i]));
-				RETURN_IF_ERROR(Reader.Read(boxData.display_primaries_y[i]));
-			}
-			RETURN_IF_ERROR(Reader.Read(boxData.white_point_x));
-			RETURN_IF_ERROR(Reader.Read(boxData.white_point_y));
-			RETURN_IF_ERROR(Reader.Read(boxData.max_display_mastering_luminance));
-			RETURN_IF_ERROR(Reader.Read(boxData.min_display_mastering_luminance));
-			const float ST2086_Chroma_Scale = 50000.0f;
-			const float ST2086_Luma_Scale = 10000.0f;
-			FVideoDecoderHDRMetadata_mastering_display_colour_volume mdcv;
-			// The order in the MDCV box is G,B,R
-			mdcv.display_primaries_x[0] = boxData.display_primaries_x[2] / ST2086_Chroma_Scale;
-			mdcv.display_primaries_y[0] = boxData.display_primaries_y[2] / ST2086_Chroma_Scale;
-			mdcv.display_primaries_x[1] = boxData.display_primaries_x[0] / ST2086_Chroma_Scale;
-			mdcv.display_primaries_y[1] = boxData.display_primaries_y[0] / ST2086_Chroma_Scale;
-			mdcv.display_primaries_x[2] = boxData.display_primaries_x[1] / ST2086_Chroma_Scale;
-			mdcv.display_primaries_y[2] = boxData.display_primaries_y[1] / ST2086_Chroma_Scale;
-			mdcv.white_point_x = boxData.white_point_x / ST2086_Chroma_Scale;
-			mdcv.white_point_y = boxData.white_point_y / ST2086_Chroma_Scale;
-			mdcv.max_display_mastering_luminance = boxData.max_display_mastering_luminance / ST2086_Luma_Scale;
-			mdcv.min_display_mastering_luminance = boxData.min_display_mastering_luminance / ST2086_Luma_Scale;
-			OutInfo.GetCodecVideoColorInfo().MDCV = mdcv;
-		}
-		return Error;
-	}
-
-
 	UEMediaError FParserISO14496_12::ParseAVCxSampleType(IPlayerSessionServices* PlayerSession, FTrack* Track, const FMP4Box* SampleBox, int32 InAvcType)
 	{
 		check(InAvcType == 1 || InAvcType == 3);
@@ -7588,26 +7474,6 @@ namespace Electra
 					{
 						break;
 					}
-					case FMP4Box::kBox_colr:
-					{
-						ParseCOLRBox(Track->CodecInformation, SampleEntry);
-						break;
-					}
-					case FMP4Box::kBox_clli:
-					{
-						ParseCOLLBox(Track->CodecInformation, SampleEntry, false);
-						break;
-					}
-					case FMP4Box::kBox_coll:
-					{
-						ParseCOLLBox(Track->CodecInformation, SampleEntry, true);
-						break;
-					}
-					case FMP4Box::kBox_mdcv:
-					{
-						ParseMDCVBox(Track->CodecInformation, SampleEntry);
-						break;
-					}
 					default:
 					{
 						break;
@@ -7706,26 +7572,6 @@ namespace Electra
 					{
 						break;
 					}
-					case FMP4Box::kBox_colr:
-					{
-						ParseCOLRBox(Track->CodecInformation, SampleEntry);
-						break;
-					}
-					case FMP4Box::kBox_clli:
-					{
-						ParseCOLLBox(Track->CodecInformation, SampleEntry, false);
-						break;
-					}
-					case FMP4Box::kBox_coll:
-					{
-						ParseCOLLBox(Track->CodecInformation, SampleEntry, true);
-						break;
-					}
-					case FMP4Box::kBox_mdcv:
-					{
-						ParseMDCVBox(Track->CodecInformation, SampleEntry);
-						break;
-					}
 					default:
 					{
 						break;
@@ -8096,19 +7942,6 @@ namespace Electra
 							}
 							default:
 							{
-								if (SampleEntry->GetType() == FMP4Box::kBox_colr)
-								{
-									ParseCOLRBox(CodecInformation, SampleEntry);
-								}
-								else if (SampleEntry->GetType() == FMP4Box::kBox_clli || SampleEntry->GetType() == FMP4Box::kBox_coll)
-								{
-									ParseCOLLBox(CodecInformation, SampleEntry, SampleEntry->GetType() == FMP4Box::kBox_coll);
-								}
-								else if (SampleEntry->GetType() == FMP4Box::kBox_mdcv)
-								{
-									ParseMDCVBox(CodecInformation, SampleEntry);
-								}
-
 								// Note: adding the box via the filler data only works as long as these boxes are not explicitly handled in this module!
 								FString BoxName = Utils::Printable4CC(SampleEntry->GetType());
 								BoxName.Append(TEXT("_box"));
diff --git a/Engine/Plugins/Media/ElectraPlayer/Source/ElectraPlayerRuntime/Private/Runtime/Demuxer/ParserMKV.cpp b/Engine/Plugins/Media/ElectraPlayer/Source/ElectraPlayerRuntime/Private/Runtime/Demuxer/ParserMKV.cpp
--- a/Engine/Plugins/Media/ElectraPlayer/Source/ElectraPlayerRuntime/Private/Runtime/Demuxer/ParserMKV.cpp
+++ b/Engine/Plugins/Media/ElectraPlayer/Source/ElectraPlayerRuntime/Private/Runtime/Demuxer/ParserMKV.cpp
@@ -18,6 +18,7 @@
 // Needed to extract information from headers of AVC, HEVC and AAC files.
 #include "Utils/MPEG/ElectraUtilsMPEGVideo_H264.h"
 #include "Utils/MPEG/ElectraUtilsMPEGVideo_H265.h"
+#include "Utils/MPEG/ElectraUtilsMPEGVideo.h"
 #include "Utils/MPEG/ElectraUtilsMPEGAudio.h"
 
 DECLARE_LOG_CATEGORY_EXTERN(LogElectraMKVParser, Log, All);
@@ -1119,6 +1120,8 @@ namespace Electra
 			{ return FrameRate; }
 			TMKVElementPtr<FMKVColour> GetColours() const
 			{ return Colours.Num() ? Colours[0] : nullptr; }
+			int32 GetAlphaMode() const
+			{ return (int32)AlphaMode; }
 		private:
 			TArray<TMKVElementPtr<FMKVColour>> Colours;
 			TArray<uint8> UncompressedFourCC;
@@ -2735,6 +2738,76 @@ namespace Electra
 		}
 		FString Codec4CC(*Codec4CCPtr);
 
+		struct FCodecVideoColorInfo
+		{
+			TOptional<int32> ColourPrimaries;
+			TOptional<int32> TransferCharacteristics;
+			TOptional<int32> MatrixCoefficients;
+			TOptional<int32> VideoFullRangeFlag;
+			TOptional<int32> VideoFormat;
+			TOptional<int32> BitDepthLuma;
+			TOptional<int32> ChromaSubsampling;
+		};
+
+		auto ParseColorInformation = [](FStreamCodecInformation& InOutCodecInformation, FCodecVideoColorInfo& InOutCodecVideoColorInfo, const TMKVElementPtr<FMKVColour>& InColours) -> void
+		{
+			if (InColours.IsValid())
+			{
+				InOutCodecVideoColorInfo.ColourPrimaries = InColours->GetPrimaries();
+				InOutCodecVideoColorInfo.TransferCharacteristics = InColours->GetTransferCharacteristics();
+				InOutCodecVideoColorInfo.MatrixCoefficients = InColours->GetMatrixCoefficients();
+				if ((!InOutCodecVideoColorInfo.BitDepthLuma.IsSet() || InOutCodecVideoColorInfo.BitDepthLuma.GetValue() == 0) && InColours->GetBitsPerChannel())
+				{
+					InOutCodecVideoColorInfo.BitDepthLuma = InColours->GetBitsPerChannel();
+				}
+				ElectraDecodersUtil::MPEG::FCommonColorimetry colr;
+				colr.video_format = 5;	// unspecified
+				colr.colour_primaries = InColours->GetPrimaries();
+				colr.transfer_characteristics = InColours->GetTransferCharacteristics();
+				colr.matrix_coeffs = InColours->GetMatrixCoefficients();
+				if (InColours->GetRange() == 2)
+				{
+					InOutCodecVideoColorInfo.VideoFullRangeFlag = 1;
+					colr.video_full_range_flag = 1;
+				}
+				InOutCodecInformation.GetExtras().Set(TEXT("colr_box"), FVariantValue(ElectraDecodersUtil::MPEG::BuildCOLRBox(colr)));
+
+				// Color light level values present?
+				if (InColours->GetMaxCLL() && InColours->GetMaxFALL())
+				{
+					ElectraDecodersUtil::MPEG::FSEIcontent_light_level_info clli;
+					clli.max_content_light_level = InColours->GetMaxCLL();
+					clli.max_pic_average_light_level = InColours->GetMaxFALL();
+					InOutCodecInformation.GetExtras().Set(TEXT("clli_box"), FVariantValue(ElectraDecodersUtil::MPEG::BuildCLLIBox(clli)));
+				}
+				
+				// Mastering display color volume present?
+				TMKVElementPtr<FMKVMasteringMetadata> MasteringMetadata = InColours->GetMasteringMetadata();
+				if (MasteringMetadata.IsValid())
+				{
+					const double ST2086_Chroma_Scale = 50000.0f;
+					const double ST2086_Luma_Scale = 10000.0f;
+					ElectraDecodersUtil::MPEG::FSEImastering_display_colour_volume mdcv;
+					mdcv.display_primaries_x[0] = (uint16)(FMath::RoundToInt32(MasteringMetadata->GetPrimaryRChromaticityX() * ST2086_Chroma_Scale));
+					mdcv.display_primaries_y[0] = (uint16)(FMath::RoundToInt32(MasteringMetadata->GetPrimaryRChromaticityY() * ST2086_Chroma_Scale));
+					mdcv.display_primaries_x[1] = (uint16)(FMath::RoundToInt32(MasteringMetadata->GetPrimaryGChromaticityX() * ST2086_Chroma_Scale));
+					mdcv.display_primaries_y[1] = (uint16)(FMath::RoundToInt32(MasteringMetadata->GetPrimaryGChromaticityY() * ST2086_Chroma_Scale));
+					mdcv.display_primaries_x[2] = (uint16)(FMath::RoundToInt32(MasteringMetadata->GetPrimaryBChromaticityX() * ST2086_Chroma_Scale));
+					mdcv.display_primaries_y[2] = (uint16)(FMath::RoundToInt32(MasteringMetadata->GetPrimaryBChromaticityY() * ST2086_Chroma_Scale));
+					mdcv.white_point_x = (uint16)(FMath::RoundToInt32(MasteringMetadata->GetWhitePointChromaticityX() * ST2086_Chroma_Scale));
+					mdcv.white_point_y = (uint16)(FMath::RoundToInt32(MasteringMetadata->GetWhitePointChromaticityY() * ST2086_Chroma_Scale));
+					mdcv.max_display_mastering_luminance = (uint32)(FMath::RoundToInt32(MasteringMetadata->GetLuminanceMax() * ST2086_Luma_Scale));
+					mdcv.min_display_mastering_luminance = (uint32)(FMath::RoundToInt32(MasteringMetadata->GetLuminanceMin() * ST2086_Luma_Scale));
+					InOutCodecInformation.GetExtras().Set(TEXT("mdcv_box"), FVariantValue(ElectraDecodersUtil::MPEG::BuildMDCVBox(mdcv)));
+				}
+			}
+		};
+
+
+
+
+
+
 		// AVC / H.264 ?
 		if (InFromTrack->GetTrackType() == EMKVTrackType::Video && Codec4CC.Equals(TEXT("avc")))
 		{
@@ -2892,43 +2965,10 @@ namespace Electra
 				return false;
 			}
 
-			TMKVElementPtr<FMKVColour> Colours = Video->GetColours();
-			if (Colours.IsValid())
-			{
-				OutCodecInformation.GetCodecVideoColorInfo().ColourPrimaries = Colours->GetPrimaries();
-				OutCodecInformation.GetCodecVideoColorInfo().TransferCharacteristics = Colours->GetTransferCharacteristics();
-				OutCodecInformation.GetCodecVideoColorInfo().MatrixCoefficients = Colours->GetMatrixCoefficients();
-				if (Colours->GetRange() == 2)
-				{
-					OutCodecInformation.GetCodecVideoColorInfo().VideoFullRangeFlag = 1;
-				}
-				//OutCodecInformation.GetCodecVideoColorInfo().VideoFormat = Colours->GetPrimaries();
-				if ((!OutCodecInformation.GetCodecVideoColorInfo().BitDepthLuma.IsSet() || OutCodecInformation.GetCodecVideoColorInfo().BitDepthLuma.GetValue() == 0) && Colours->GetBitsPerChannel())
-				{
-					OutCodecInformation.GetCodecVideoColorInfo().BitDepthLuma = Colours->GetBitsPerChannel();
-				}
-				// Color light level values present?
-				if (Colours->GetMaxCLL() && Colours->GetMaxFALL())
-				{
-					OutCodecInformation.GetCodecVideoColorInfo().CLLI = { Colours->GetMaxCLL(), Colours->GetMaxFALL() };
-				}
-				TMKVElementPtr<FMKVMasteringMetadata> MasteringMetadata = Colours->GetMasteringMetadata();
-				if (MasteringMetadata.IsValid())
-				{
-					FVideoDecoderHDRMetadata_mastering_display_colour_volume mdcv;
-					mdcv.display_primaries_x[0] = (float) MasteringMetadata->GetPrimaryRChromaticityX();
-					mdcv.display_primaries_y[0] = (float) MasteringMetadata->GetPrimaryRChromaticityY();
-					mdcv.display_primaries_x[1] = (float) MasteringMetadata->GetPrimaryGChromaticityX();
-					mdcv.display_primaries_y[1] = (float) MasteringMetadata->GetPrimaryGChromaticityY();
-					mdcv.display_primaries_x[2] = (float) MasteringMetadata->GetPrimaryBChromaticityX();
-					mdcv.display_primaries_y[2] = (float) MasteringMetadata->GetPrimaryBChromaticityY();
-					mdcv.white_point_x = (float) MasteringMetadata->GetWhitePointChromaticityX();
-					mdcv.white_point_y = (float) MasteringMetadata->GetWhitePointChromaticityY();
-					mdcv.max_display_mastering_luminance = (float) MasteringMetadata->GetLuminanceMax();
-					mdcv.min_display_mastering_luminance = (float) MasteringMetadata->GetLuminanceMin();
-					OutCodecInformation.GetCodecVideoColorInfo().MDCV = mdcv;
-				}
-			}
+			// Parse color information
+			FCodecVideoColorInfo CodecVideoColorInfo;
+			ParseColorInformation(OutCodecInformation, CodecVideoColorInfo, Video->GetColours());
+			OutCodecInformation.GetExtras().Set(TEXT("alphamode"), FVariantValue((int64) Video->GetAlphaMode()));
 
 			OutCodecInformation.SetStreamType(EStreamType::Video);
 			OutCodecInformation.SetMimeType(TEXT("video/VP8"));
@@ -2937,24 +2977,24 @@ namespace Electra
 			OutCodecInformation.SetCodec4CC(Utils::Make4CC('v','p','0','8'));
 			FString crfc = FString::Printf(TEXT("vp08.%02d.%02d.%02d.%02d.%02d.%02d.%02d.%02d"),
 				OutCodecInformation.GetProfile(), OutCodecInformation.GetProfileLevel(),
-				OutCodecInformation.GetCodecVideoColorInfo().BitDepthLuma.Get(8),
-				OutCodecInformation.GetCodecVideoColorInfo().ChromaSubsampling.Get(0),
-				OutCodecInformation.GetCodecVideoColorInfo().ColourPrimaries.Get(2),
-				OutCodecInformation.GetCodecVideoColorInfo().TransferCharacteristics.Get(2),
-				OutCodecInformation.GetCodecVideoColorInfo().MatrixCoefficients.Get(2),
-				OutCodecInformation.GetCodecVideoColorInfo().VideoFullRangeFlag.Get(0));
+				CodecVideoColorInfo.BitDepthLuma.Get(8),
+				CodecVideoColorInfo.ChromaSubsampling.Get(0),
+				CodecVideoColorInfo.ColourPrimaries.Get(2),
+				CodecVideoColorInfo.TransferCharacteristics.Get(2),
+				CodecVideoColorInfo.MatrixCoefficients.Get(2),
+				CodecVideoColorInfo.VideoFullRangeFlag.Get(0));
 			OutCodecInformation.SetCodecSpecifierRFC6381(crfc);
 
 			// Construct a VPCodecConfigurationBox
 			TArray<uint8> vpcC { 1,0,0,0, 0,0,0,0, 0,0,0,0 };
 			vpcC[4] = (uint8)OutCodecInformation.GetProfile();
 			vpcC[5] = (uint8)OutCodecInformation.GetProfileLevel();
-			vpcC[6] = (uint8)((OutCodecInformation.GetCodecVideoColorInfo().BitDepthLuma.Get(8) << 4) +
-							  (OutCodecInformation.GetCodecVideoColorInfo().ChromaSubsampling.Get(0) << 1) +
-							   OutCodecInformation.GetCodecVideoColorInfo().VideoFullRangeFlag.Get(0));
-			vpcC[7] = (uint8)OutCodecInformation.GetCodecVideoColorInfo().ColourPrimaries.Get(2);
-			vpcC[8] = (uint8)OutCodecInformation.GetCodecVideoColorInfo().TransferCharacteristics.Get(2);
-			vpcC[9] = (uint8)OutCodecInformation.GetCodecVideoColorInfo().MatrixCoefficients.Get(2);
+			vpcC[6] = (uint8)((CodecVideoColorInfo.BitDepthLuma.Get(8) << 4) +
+							  (CodecVideoColorInfo.ChromaSubsampling.Get(0) << 1) +
+							   CodecVideoColorInfo.VideoFullRangeFlag.Get(0));
+			vpcC[7] = (uint8)CodecVideoColorInfo.ColourPrimaries.Get(2);
+			vpcC[8] = (uint8)CodecVideoColorInfo.TransferCharacteristics.Get(2);
+			vpcC[9] = (uint8)CodecVideoColorInfo.MatrixCoefficients.Get(2);
 			OutCodecInformation.GetExtras().Set(StreamCodecInformationOptions::VPccBox, FVariantValue(vpcC));
 
 			int32 CropL = Video->GetCropLeft();
@@ -2992,6 +3032,8 @@ namespace Electra
 				return false;
 			}
 
+			FCodecVideoColorInfo CodecVideoColorInfo;
+
 			// The codec private data, if present, should be this: https://www.webmproject.org/docs/container/#vp9-codec-feature-metadata-codecprivate
 			if (InFromTrack->GetCodecPrivate().Num())
 			{
@@ -3003,12 +3045,12 @@ namespace Electra
 				{
 					OutCodecInformation.SetProfile(cp[0]);
 					OutCodecInformation.SetProfileLevel(cp[1]);
-					OutCodecInformation.GetCodecVideoColorInfo().BitDepthLuma = cp[2] >> 4;
-					OutCodecInformation.GetCodecVideoColorInfo().ChromaSubsampling = (cp[2] >> 3) & 7;
-					OutCodecInformation.GetCodecVideoColorInfo().VideoFullRangeFlag = cp[2] & 1;
-					OutCodecInformation.GetCodecVideoColorInfo().ColourPrimaries = cp[3];
-					OutCodecInformation.GetCodecVideoColorInfo().TransferCharacteristics = cp[4];
-					OutCodecInformation.GetCodecVideoColorInfo().MatrixCoefficients = cp[5];
+					CodecVideoColorInfo.BitDepthLuma = cp[2] >> 4;
+					CodecVideoColorInfo.ChromaSubsampling = (cp[2] >> 3) & 7;
+					CodecVideoColorInfo.VideoFullRangeFlag = cp[2] & 1;
+					CodecVideoColorInfo.ColourPrimaries = cp[3];
+					CodecVideoColorInfo.TransferCharacteristics = cp[4];
+					CodecVideoColorInfo.MatrixCoefficients = cp[5];
 				}
 				else
 				{
@@ -3022,9 +3064,9 @@ namespace Electra
 							// Level
 							case 2: if (*cp++ == 1) OutCodecInformation.SetProfileLevel(*cp++); else return false; break;
 							// Bit depth
-							case 3: if (*cp++ == 1) OutCodecInformation.GetCodecVideoColorInfo().BitDepthLuma = *cp++; else return false; break;
+							case 3: if (*cp++ == 1) CodecVideoColorInfo.BitDepthLuma = *cp++; else return false; break;
 							// Chroma subsampling
-							case 4: if (*cp++ == 1) OutCodecInformation.GetCodecVideoColorInfo().ChromaSubsampling = *cp++; else return false; break;
+							case 4: if (*cp++ == 1) CodecVideoColorInfo.ChromaSubsampling = *cp++; else return false; break;
 							default:
 							{
 								uint8 len = *cp++;
@@ -3036,43 +3078,9 @@ namespace Electra
 				}
 			}
 
-			TMKVElementPtr<FMKVColour> Colours = Video->GetColours();
-			if (Colours.IsValid())
-			{
-				OutCodecInformation.GetCodecVideoColorInfo().ColourPrimaries = Colours->GetPrimaries();
-				OutCodecInformation.GetCodecVideoColorInfo().TransferCharacteristics = Colours->GetTransferCharacteristics();
-				OutCodecInformation.GetCodecVideoColorInfo().MatrixCoefficients = Colours->GetMatrixCoefficients();
-				if (Colours->GetRange() == 2)
-				{
-					OutCodecInformation.GetCodecVideoColorInfo().VideoFullRangeFlag = 1;
-				}
-				//OutCodecInformation.GetCodecVideoColorInfo().VideoFormat = Colours->GetPrimaries();
-				if ((!OutCodecInformation.GetCodecVideoColorInfo().BitDepthLuma.IsSet() || OutCodecInformation.GetCodecVideoColorInfo().BitDepthLuma.GetValue() == 0) && Colours->GetBitsPerChannel())
-				{
-					OutCodecInformation.GetCodecVideoColorInfo().BitDepthLuma = Colours->GetBitsPerChannel();
-				}
-				// Color light level values present?
-				if (Colours->GetMaxCLL() && Colours->GetMaxFALL())
-				{
-					OutCodecInformation.GetCodecVideoColorInfo().CLLI = { Colours->GetMaxCLL(), Colours->GetMaxFALL() };
-				}
-				TMKVElementPtr<FMKVMasteringMetadata> MasteringMetadata = Colours->GetMasteringMetadata();
-				if (MasteringMetadata.IsValid())
-				{
-					FVideoDecoderHDRMetadata_mastering_display_colour_volume mdcv;
-					mdcv.display_primaries_x[0] = (float) MasteringMetadata->GetPrimaryRChromaticityX();
-					mdcv.display_primaries_y[0] = (float) MasteringMetadata->GetPrimaryRChromaticityY();
-					mdcv.display_primaries_x[1] = (float) MasteringMetadata->GetPrimaryGChromaticityX();
-					mdcv.display_primaries_y[1] = (float) MasteringMetadata->GetPrimaryGChromaticityY();
-					mdcv.display_primaries_x[2] = (float) MasteringMetadata->GetPrimaryBChromaticityX();
-					mdcv.display_primaries_y[2] = (float) MasteringMetadata->GetPrimaryBChromaticityY();
-					mdcv.white_point_x = (float) MasteringMetadata->GetWhitePointChromaticityX();
-					mdcv.white_point_y = (float) MasteringMetadata->GetWhitePointChromaticityY();
-					mdcv.max_display_mastering_luminance = (float) MasteringMetadata->GetLuminanceMax();
-					mdcv.min_display_mastering_luminance = (float) MasteringMetadata->GetLuminanceMin();
-					OutCodecInformation.GetCodecVideoColorInfo().MDCV = mdcv;
-				}
-			}
+			// Parse color information
+			ParseColorInformation(OutCodecInformation, CodecVideoColorInfo, Video->GetColours());
+			OutCodecInformation.GetExtras().Set(TEXT("alphamode"), FVariantValue((int64) Video->GetAlphaMode()));
 
 			OutCodecInformation.SetStreamType(EStreamType::Video);
 			OutCodecInformation.SetMimeType(TEXT("video/VP9"));
@@ -3081,24 +3089,24 @@ namespace Electra
 			OutCodecInformation.SetCodec4CC(Utils::Make4CC('v','p','0','9'));
 			FString crfc = FString::Printf(TEXT("vp09.%02d.%02d.%02d.%02d.%02d.%02d.%02d.%02d"),
 				OutCodecInformation.GetProfile(), OutCodecInformation.GetProfileLevel(),
-				OutCodecInformation.GetCodecVideoColorInfo().BitDepthLuma.Get(8),
-				OutCodecInformation.GetCodecVideoColorInfo().ChromaSubsampling.Get(0),
-				OutCodecInformation.GetCodecVideoColorInfo().ColourPrimaries.Get(2),
-				OutCodecInformation.GetCodecVideoColorInfo().TransferCharacteristics.Get(2),
-				OutCodecInformation.GetCodecVideoColorInfo().MatrixCoefficients.Get(2),
-				OutCodecInformation.GetCodecVideoColorInfo().VideoFullRangeFlag.Get(0));
+				CodecVideoColorInfo.BitDepthLuma.Get(8),
+				CodecVideoColorInfo.ChromaSubsampling.Get(0),
+				CodecVideoColorInfo.ColourPrimaries.Get(2),
+				CodecVideoColorInfo.TransferCharacteristics.Get(2),
+				CodecVideoColorInfo.MatrixCoefficients.Get(2),
+				CodecVideoColorInfo.VideoFullRangeFlag.Get(0));
 			OutCodecInformation.SetCodecSpecifierRFC6381(crfc);
 
 			// Construct a VPCodecConfigurationBox
 			TArray<uint8> vpcC { 1,0,0,0, 0,0,0,0, 0,0,0,0 };
 			vpcC[4] = (uint8)OutCodecInformation.GetProfile();
 			vpcC[5] = (uint8)OutCodecInformation.GetProfileLevel();
-			vpcC[6] = (uint8)((OutCodecInformation.GetCodecVideoColorInfo().BitDepthLuma.Get(8) << 4) +
-							  (OutCodecInformation.GetCodecVideoColorInfo().ChromaSubsampling.Get(0) << 1) +
-							   OutCodecInformation.GetCodecVideoColorInfo().VideoFullRangeFlag.Get(0));
-			vpcC[7] = (uint8)OutCodecInformation.GetCodecVideoColorInfo().ColourPrimaries.Get(2);
-			vpcC[8] = (uint8)OutCodecInformation.GetCodecVideoColorInfo().TransferCharacteristics.Get(2);
-			vpcC[9] = (uint8)OutCodecInformation.GetCodecVideoColorInfo().MatrixCoefficients.Get(2);
+			vpcC[6] = (uint8)((CodecVideoColorInfo.BitDepthLuma.Get(8) << 4) +
+							  (CodecVideoColorInfo.ChromaSubsampling.Get(0) << 1) +
+							   CodecVideoColorInfo.VideoFullRangeFlag.Get(0));
+			vpcC[7] = (uint8)CodecVideoColorInfo.ColourPrimaries.Get(2);
+			vpcC[8] = (uint8)CodecVideoColorInfo.TransferCharacteristics.Get(2);
+			vpcC[9] = (uint8)CodecVideoColorInfo.MatrixCoefficients.Get(2);
 			OutCodecInformation.GetExtras().Set(StreamCodecInformationOptions::VPccBox, FVariantValue(vpcC));
 
 			int32 CropL = Video->GetCropLeft();
diff --git a/Engine/Plugins/Media/ElectraPlayer/Source/ElectraPlayerRuntime/Private/Runtime/Player/StreamSegmentReaderCommon.cpp b/Engine/Plugins/Media/ElectraPlayer/Source/ElectraPlayerRuntime/Private/Runtime/Player/StreamSegmentReaderCommon.cpp
--- a/Engine/Plugins/Media/ElectraPlayer/Source/ElectraPlayerRuntime/Private/Runtime/Player/StreamSegmentReaderCommon.cpp
+++ b/Engine/Plugins/Media/ElectraPlayer/Source/ElectraPlayerRuntime/Private/Runtime/Player/StreamSegmentReaderCommon.cpp
@@ -2338,11 +2338,26 @@ FStreamSegmentReaderCommon::FStreamHandler::EHandleResult FStreamSegmentReaderCo
 
 						// Any additional data?
 						const TMap<uint64, TArray<uint8>>& BlockAdditionalData = Action->GetBlockAdditionalData();
-						// What type of additional data is there?
-						if (BlockAdditionalData.Contains(4))	// VP9 ITU T.35 metadata?
+						for(auto& badIt : BlockAdditionalData)
 						{
-							AccessUnit->DynamicSidebandData = MakeUnique<TMap<uint32, TArray<uint8>>>();
-							AccessUnit->DynamicSidebandData->Emplace(Utils::Make4CC('i','t','3','5'), BlockAdditionalData[4]);
+							auto AddDynamicBlockData = [&](const FName& InName, const TArray<uint8>& InData) -> void
+							{
+								if (!AccessUnit->DynamicSidebandData.IsValid())
+								{
+									AccessUnit->DynamicSidebandData = MakeUnique<TMap<FName, TArray<uint8>>>();
+								}
+								AccessUnit->DynamicSidebandData->Emplace(InName, InData);
+							};
+
+							// What type of additional data is there?
+							if (BlockAdditionalData.Contains(1))		// VP9 alpha channel?
+							{
+								AddDynamicBlockData(DynamicSidebandData::VPxAlpha, BlockAdditionalData[1]);
+							}
+							else if (BlockAdditionalData.Contains(4))	// VP9 ITU T.35 metadata?
+							{
+								AddDynamicBlockData(DynamicSidebandData::ITU_T_35, BlockAdditionalData[4]);
+							}
 						}
 					}
 					// VP8 codec?
diff --git a/Engine/Plugins/Media/ElectraPlayer/Source/ElectraPlayerRuntime/Private/Runtime/Player/mkv/StreamReaderMKV.cpp b/Engine/Plugins/Media/ElectraPlayer/Source/ElectraPlayerRuntime/Private/Runtime/Player/mkv/StreamReaderMKV.cpp
--- a/Engine/Plugins/Media/ElectraPlayer/Source/ElectraPlayerRuntime/Private/Runtime/Player/mkv/StreamReaderMKV.cpp
+++ b/Engine/Plugins/Media/ElectraPlayer/Source/ElectraPlayerRuntime/Private/Runtime/Player/mkv/StreamReaderMKV.cpp
@@ -518,11 +518,26 @@ void FStreamReaderMKV::HandleRequest()
 
 						// Any additional data?
 						const TMap<uint64, TArray<uint8>>& BlockAdditionalData = Action->GetBlockAdditionalData();
-						// What type of additional data is there?
-						if (BlockAdditionalData.Contains(4))	// VP9 ITU T.35 metadata?
+						for(auto& badIt : BlockAdditionalData)
 						{
-							AccessUnit->DynamicSidebandData = MakeUnique<TMap<uint32, TArray<uint8>>>();
-							AccessUnit->DynamicSidebandData->Emplace(Utils::Make4CC('i','t','3','5'), BlockAdditionalData[4]);
+							auto AddDynamicBlockData = [&](const FName& InName, const TArray<uint8>& InData) -> void
+							{
+								if (!AccessUnit->DynamicSidebandData.IsValid())
+								{
+									AccessUnit->DynamicSidebandData = MakeUnique<TMap<FName, TArray<uint8>>>();
+								}
+								AccessUnit->DynamicSidebandData->Emplace(InName, InData);
+							};
+
+							// What type of additional data is there?
+							if (BlockAdditionalData.Contains(1))		// VP9 alpha channel?
+							{
+								AddDynamicBlockData(DynamicSidebandData::VPxAlpha, BlockAdditionalData[1]);
+							}
+							else if (BlockAdditionalData.Contains(4))	// VP9 ITU T.35 metadata?
+							{
+								AddDynamicBlockData(DynamicSidebandData::ITU_T_35, BlockAdditionalData[4]);
+							}
 						}
 					}
 					// VP8 codec?
diff --git a/Engine/Plugins/Media/ElectraPlayer/Source/ElectraPlayerRuntime/Private/Runtime/StreamAccessUnitBuffer.h b/Engine/Plugins/Media/ElectraPlayer/Source/ElectraPlayerRuntime/Private/Runtime/StreamAccessUnitBuffer.h
--- a/Engine/Plugins/Media/ElectraPlayer/Source/ElectraPlayerRuntime/Private/Runtime/StreamAccessUnitBuffer.h
+++ b/Engine/Plugins/Media/ElectraPlayer/Source/ElectraPlayerRuntime/Private/Runtime/StreamAccessUnitBuffer.h
@@ -12,6 +12,11 @@
 
 namespace Electra
 {
+	namespace DynamicSidebandData
+	{
+		const FName ITU_T_35(TEXT("itu-t-35"));
+		const FName VPxAlpha(TEXT("vpx-alpha"));
+	}
 
 
 	class IAccessUnitMemoryProvider
@@ -63,7 +68,7 @@ namespace Electra
 
 		TSharedPtrTS<CodecData>		AUCodecData;					//!< If set, points to constant sideband data for this access unit.
 		TSharedPtrTS<const FBufferSourceInfo>	BufferSourceInfo;
-		TUniquePtr<TMap<uint32, TArray<uint8>>> DynamicSidebandData;	//!< If set, contains a map of dynamically changing sideband data for just this access unit.
+		TUniquePtr<TMap<FName, TArray<uint8>>> DynamicSidebandData;	//!< If set, contains a map of dynamically changing sideband data for just this access unit.
 		EStreamType					ESType;							//!< Type of elementary stream this is an access unit of.
 		FTimeValue					PTS;							//!< PTS
 		FTimeValue					DTS;							//!< DTS
diff --git a/Engine/Plugins/Media/ElectraPlayer/Source/ElectraPlayerRuntime/Private/Runtime/StreamTypes.h b/Engine/Plugins/Media/ElectraPlayer/Source/ElectraPlayerRuntime/Private/Runtime/StreamTypes.h
--- a/Engine/Plugins/Media/ElectraPlayer/Source/ElectraPlayerRuntime/Private/Runtime/StreamTypes.h
+++ b/Engine/Plugins/Media/ElectraPlayer/Source/ElectraPlayerRuntime/Private/Runtime/StreamTypes.h
@@ -529,12 +529,7 @@ namespace Electra
 		{
 			StreamLanguageTag = InStreamLanguageTag;
 		}
-/*
-		const FString& GetStreamLanguageCode() const
-		{
-			return StreamLanguageCode;
-		}
-*/
+
 		void SetCodecSpecificData(const TArray<uint8>& InCSD)
 		{
 			CSD = InCSD;
@@ -599,7 +594,6 @@ namespace Electra
 			Extras.Clear();
 			CSD.Empty();
 			DCR.Empty();
-			CodecVideoColorInfo.Reset();
 			ResetHumanReadableCodecName();
 		}
 
@@ -632,32 +626,6 @@ namespace Electra
 			return false;
 		}
 
-		struct FCodecVideoColorInfo
-		{
-			void Reset()
-			{
-				ColourPrimaries.Reset();
-				TransferCharacteristics.Reset();
-				MatrixCoefficients.Reset();
-				VideoFullRangeFlag.Reset();
-				VideoFormat.Reset();
-				BitDepthLuma.Reset();
-				ChromaSubsampling.Reset();
-			}
-			TOptional<int32> ColourPrimaries;
-			TOptional<int32> TransferCharacteristics;
-			TOptional<int32> MatrixCoefficients;
-			TOptional<int32> VideoFullRangeFlag;
-			TOptional<int32> VideoFormat;
-			TOptional<int32> BitDepthLuma;
-			TOptional<int32> ChromaSubsampling;
-			TOptional<FVideoDecoderHDRMetadata_content_light_level_info> CLLI;
-			TOptional<FVideoDecoderHDRMetadata_mastering_display_colour_volume> MDCV;
-		};
-		FCodecVideoColorInfo& GetCodecVideoColorInfo()
-		{ return CodecVideoColorInfo; }
-
-
 	private:
 		void ResetHumanReadableCodecName()
 		{
@@ -720,7 +688,6 @@ namespace Electra
 		FParamDict		Extras;						//!< Additional details/properties, depending on playlist.
 		TArray<uint8>	CSD;						//!< Codec specific data, if available.
 		TArray<uint8>	DCR;						//!< Decoder configuration record, if available.
-		FCodecVideoColorInfo CodecVideoColorInfo;
 		mutable FString	HumanReadableCodecName;
 	};
 
diff --git a/Engine/Plugins/Media/VPxDecoderElectra/Source/VPxDecoderElectra/Private/VPxDecoder/ElectraMediaVPxDecoder.cpp b/Engine/Plugins/Media/VPxDecoderElectra/Source/VPxDecoderElectra/Private/VPxDecoder/ElectraMediaVPxDecoder.cpp
--- a/Engine/Plugins/Media/VPxDecoderElectra/Source/VPxDecoderElectra/Private/VPxDecoder/ElectraMediaVPxDecoder.cpp
+++ b/Engine/Plugins/Media/VPxDecoderElectra/Source/VPxDecoderElectra/Private/VPxDecoder/ElectraMediaVPxDecoder.cpp
@@ -19,7 +19,14 @@
 #include "Containers/Queue.h"
 #include <Stats/Stats.h>
 
+
+#define ENABLE_GPU_BUFFER_HELPER 1
+
+#if ENABLE_GPU_BUFFER_HELPER
 #include COMPILED_PLATFORM_HEADER(ElectraDecoderGPUBufferHelpers.h)
+#else
+#define ELECTRA_MEDIAGPUBUFFER_DX12 0
+#endif
 
 #include <vpx/vpx_decoder.h>
 #include <vpx/vp8dx.h>
@@ -44,6 +51,12 @@ DECLARE_CYCLE_STAT(TEXT("ElectraDecoder ConvertOutput"), STAT_ElectraDecoder_Con
 /*********************************************************************************************************************/
 /*********************************************************************************************************************/
 
+namespace DynamicSidebandData
+{
+	const FName VPxAlpha(TEXT("vpx-alpha"));
+}
+
+
 class FVideoDecoderVPxElectra;
 
 
@@ -179,7 +192,6 @@ public:
 	int32 AspectH = 1;
 	int32 FrameRateN = 0;
 	int32 FrameRateD = 0;
-	int32 PixelFormat = 0;
 	TMap<FString, FVariant> ExtraValues;
 
 	uint32 Codec4CC = 0;
@@ -249,6 +261,8 @@ private:
 		TMap<FString, FVariant> AdditionalOptions;
 		void* InputDataCopy = nullptr;
 		int32 InputDataSize = 0;
+		void* InputAlphaDataCopy = nullptr;
+		int32 InputAlphaDataSize = 0;
 		int32 SuperFrameIndexPlusOne = 0;
 	};
 
@@ -278,9 +292,10 @@ private:
 		Failure,
 		GotEOS
 	};
-	EConvertResult ConvertDecoderOutput(vpx_image_t* InDecodedImage);
+	EConvertResult ConvertDecoderOutput(vpx_image_t* InDecodedImage, vpx_image_t* InDecodedAlphaImage);
 
 	bool ConvertDecodedImageToNV12orP010(FVideoDecoderOutputVPxElectra* InNewOutput, const vpx_image_t* InDecodedImage, void* PlatformDevice, int32 PlatformDeviceVersion) const;
+	bool ConvertDecodedImageWithAlpha(FVideoDecoderOutputVPxElectra* InNewOutput, const vpx_image_t* InDecodedImage, const vpx_image_t* InDecodedAlphaImage, void* PlatformDevice, int32 PlatformDeviceVersion) const;
 
 	TMap<FString, FVariant> InitialCreationOptions;
 	TWeakPtr<IElectraDecoderResourceDelegate, ESPMode::ThreadSafe> ResourceDelegate;
@@ -292,9 +307,19 @@ private:
 	EDecodeState DecodeState = EDecodeState::Decoding;
 	uint32 Codec4CC = 0;
 
-	vpx_codec_ctx_t DecoderContext;
-	vpx_codec_ctx_t* VideoDecoderHandle = nullptr;
-	vpx_codec_iter_t VideoDecoderOutputIterator = nullptr;
+	enum EDecoderType
+	{
+		Color = 0,
+		Alpha = 1
+	};
+
+	struct FDecoderStructs
+	{
+		vpx_codec_ctx_t Context;
+		vpx_codec_ctx_t* Handle = nullptr;
+		vpx_codec_iter_t OutputIterator = nullptr;
+	};
+	FDecoderStructs Decoders[2];
 
 	uint32 MaxWidth;
 	uint32 MaxHeight;
@@ -394,11 +419,12 @@ FVideoDecoderVPxElectra::FVideoDecoderVPxElectra(const TMap<FString, FVariant>&
 {
 	Codec4CC = (uint32)ElectraDecodersUtil::GetVariantValueSafeU64(InOptions, TEXT("codec_4cc"), 0);
 
-	MaxWidth = (uint32)ElectraDecodersUtil::GetVariantValueSafeU64(InOptions, TEXT("max_width"), 1920);
-	MaxHeight = (uint32)ElectraDecodersUtil::GetVariantValueSafeU64(InOptions, TEXT("max_height"), 1080);
-
+	MaxWidth = (uint32)Align(ElectraDecodersUtil::GetVariantValueSafeU64(InOptions, TEXT("max_width"), 1920), 2);
+	MaxHeight = (uint32)Align(ElectraDecodersUtil::GetVariantValueSafeU64(InOptions, TEXT("max_height"), 1080), 2);
 	MaxOutputBuffers = (uint32)ElectraDecodersUtil::GetVariantValueSafeU64(InOptions, TEXT("max_output_buffers"), 5);
+#if ELECTRA_MEDIAGPUBUFFER_DX12
 	MaxOutputBuffers += kElectraDecoderPipelineExtraFrames;
+#endif
 }
 
 FVideoDecoderVPxElectra::~FVideoDecoderVPxElectra()
@@ -436,7 +462,7 @@ void FVideoDecoderVPxElectra::Close()
 
 IElectraDecoder::ECSDCompatibility FVideoDecoderVPxElectra::IsCompatibleWith(const TMap<FString, FVariant>& CSDAndAdditionalOptions)
 {
-	if (!VideoDecoderHandle)
+	if (!Decoders[EDecoderType::Color].Handle)
 	{
 		return IElectraDecoder::ECSDCompatibility::Compatible;
 	}
@@ -499,7 +525,7 @@ IElectraDecoder::EDecoderError FVideoDecoderVPxElectra::DecodeAccessUnit(const F
 	check((InInputAccessUnit.Flags & EElectraDecoderFlags::InitCSDOnly) == EElectraDecoderFlags::None);
 
 	// Create decoder transform if necessary.
-	if (!VideoDecoderHandle && !InternalDecoderCreate(InAdditionalOptions))
+	if (!Decoders[EDecoderType::Color].Handle && !InternalDecoderCreate(InAdditionalOptions))
 	{
 		return IElectraDecoder::EDecoderError::Error;
 	}
@@ -563,6 +589,11 @@ IElectraDecoder::EDecoderError FVideoDecoderVPxElectra::DecodeNextPending()
 				PostError(0, TEXT("Failed to validate VP9 bitstream header"), ERRCODE_INTERNAL_FAILED_TO_DECODE_INPUT);
 				return IElectraDecoder::EDecoderError::Error;
 			}
+			if (In->InputAlphaDataCopy && !ElectraDecodersUtil::VPxVideo::ParseVP9UncompressedHeader(Header, In->InputAlphaDataCopy, In->InputAlphaDataSize))
+			{
+				PostError(0, TEXT("Failed to validate VP9 alpha channel bitstream header"), ERRCODE_INTERNAL_FAILED_TO_DECODE_INPUT);
+				return IElectraDecoder::EDecoderError::Error;
+			}
 			break;
 		}
 		default:
@@ -571,20 +602,33 @@ IElectraDecoder::EDecoderError FVideoDecoderVPxElectra::DecodeNextPending()
 		}
 	}
 
-	// Invoke decoder
-	vpx_codec_err_t Result = vpx_codec_decode(VideoDecoderHandle, reinterpret_cast<const uint8*>(In->InputDataCopy), (unsigned int) In->InputDataSize, In.Get(), 0);
+	// Invoke decoder for the color channel
+	vpx_codec_err_t Result = vpx_codec_decode(Decoders[EDecoderType::Color].Handle, reinterpret_cast<const uint8*>(In->InputDataCopy), (unsigned int) In->InputDataSize, In.Get(), 0);
 	if (Result != VPX_CODEC_OK)
 	{
 		PostError(Result, TEXT("Failed to decode video decoder input"), ERRCODE_INTERNAL_FAILED_TO_DECODE_INPUT);
 		return IElectraDecoder::EDecoderError::Error;
 	}
+	// Invoke decoder for alpha channel if there is one
+	if (In->InputAlphaDataCopy)
+	{
+		check(Decoders[EDecoderType::Alpha].Handle);
+		Result = vpx_codec_decode(Decoders[EDecoderType::Alpha].Handle, reinterpret_cast<const uint8*>(In->InputAlphaDataCopy), (unsigned int) In->InputAlphaDataSize, In.Get(), 0);
+		if (Result != VPX_CODEC_OK)
+		{
+			PostError(Result, TEXT("Failed to decode video decoder alpha channel input"), ERRCODE_INTERNAL_FAILED_TO_DECODE_INPUT);
+			return IElectraDecoder::EDecoderError::Error;
+		}
+	}
 
 	// Add to the list of inputs passed to the decoder.
 	InDecoderInput.Emplace(MoveTemp(In));
 
 	// Did we produce a new frame?
-	VideoDecoderOutputIterator = nullptr;
-	EConvertResult ConvResult = ConvertDecoderOutput(vpx_codec_get_frame(VideoDecoderHandle, &VideoDecoderOutputIterator));
+	Decoders[EDecoderType::Color].OutputIterator = nullptr;
+	Decoders[EDecoderType::Alpha].OutputIterator = nullptr;
+	EConvertResult ConvResult = ConvertDecoderOutput(vpx_codec_get_frame(Decoders[EDecoderType::Color].Handle, &Decoders[EDecoderType::Color].OutputIterator)
+			, Decoders[EDecoderType::Alpha].Handle ? vpx_codec_get_frame(Decoders[EDecoderType::Alpha].Handle, &Decoders[EDecoderType::Alpha].OutputIterator) : nullptr);
 	if (ConvResult == EConvertResult::Failure)
 	{
 		return IElectraDecoder::EDecoderError::Error;
@@ -612,13 +656,23 @@ bool FVideoDecoderVPxElectra::PrepareSingleFrame(const FInputAccessUnit& InInput
 
 bool FVideoDecoderVPxElectra::DecomposeSuperFrame(const FInputAccessUnit& InInputAccessUnit, const TMap<FString, FVariant>& InAdditionalOptions)
 {
-	TArray<uint32> FrameSizes;
-	if (!ElectraDecodersUtil::VPxVideo::GetVP9SuperframeSizes(FrameSizes, InInputAccessUnit.Data, InInputAccessUnit.DataSize))
+	TArray<uint8> AlphaFrameData = ElectraDecodersUtil::GetVariantValueUInt8Array(InAdditionalOptions, DynamicSidebandData::VPxAlpha.ToString());
+
+	TArray<uint32> FrameSizes, AlphaFrameSizes;
+	if (!ElectraDecodersUtil::VPxVideo::GetVP9SuperframeSizes(FrameSizes, InInputAccessUnit.Data, InInputAccessUnit.DataSize) ||
+		(AlphaFrameData.Num() && !ElectraDecodersUtil::VPxVideo::GetVP9SuperframeSizes(AlphaFrameSizes, AlphaFrameData.GetData(), AlphaFrameData.Num())))
 	{
 		return false;
 	}
+	// If there is an alpha channel then it needs to have the same number of frames as the color channel.
+	if (AlphaFrameData.Num() && FrameSizes.Num() != AlphaFrameSizes.Num())
+	{
+		return false;
+	}
+
 	// Create input frames.
 	const uint8* Data = reinterpret_cast<const uint8*>(InInputAccessUnit.Data);
+	const uint8* AlphaData = AlphaFrameData.Num() ? AlphaFrameData.GetData() : nullptr;
 	for(int32 i=0; i<FrameSizes.Num(); ++i)
 	{
 		// Add to the list of inputs passed to the decoder.
@@ -628,6 +682,13 @@ bool FVideoDecoderVPxElectra::DecomposeSuperFrame(const FInputAccessUnit& InInpu
 		In->AccessUnit = InInputAccessUnit;
 		In->InputDataSize = FrameSizes[i];
 		In->InputDataCopy = FMemory::Malloc(FrameSizes[i]);
+		if (AlphaData)
+		{
+			In->InputAlphaDataSize = AlphaFrameSizes[i];
+			In->InputAlphaDataCopy = FMemory::Malloc(AlphaFrameSizes[i]);
+			FMemory::Memcpy(In->InputAlphaDataCopy, AlphaData, AlphaFrameSizes[i]);
+			AlphaData += AlphaFrameSizes[i];
+		}
 		FMemory::Memcpy(In->InputDataCopy, Data, FrameSizes[i]);
 		Data += FrameSizes[i];
 		// Zero the input pointer and size in the copy. That data is not owned by us and it's best not to have any
@@ -654,7 +715,7 @@ IElectraDecoder::EDecoderError FVideoDecoderVPxElectra::SendEndOfData()
 	}
 
 	// If there is a transform send an end-of-stream and drain message.
-	if (VideoDecoderHandle)
+	if (Decoders[EDecoderType::Color].Handle)
 	{
 		DecodeState = EDecodeState::Draining;
 	}
@@ -668,7 +729,7 @@ IElectraDecoder::EDecoderError FVideoDecoderVPxElectra::Flush()
 	{
 		return IElectraDecoder::EDecoderError::Error;
 	}
-	if (VideoDecoderHandle)
+	if (Decoders[EDecoderType::Color].Handle)
 	{
 		InternalDecoderDestroy();
 		DecodeState = EDecodeState::Decoding;
@@ -693,9 +754,10 @@ IElectraDecoder::EOutputStatus FVideoDecoderVPxElectra::HaveOutput()
 	}
 
 	// See if there is any additional pending output.
-	if (VideoDecoderHandle && VideoDecoderOutputIterator)
+	if (Decoders[EDecoderType::Color].Handle && Decoders[EDecoderType::Color].OutputIterator)
 	{
-		EConvertResult ConvResult = ConvertDecoderOutput(vpx_codec_get_frame(VideoDecoderHandle, &VideoDecoderOutputIterator));
+		EConvertResult ConvResult = ConvertDecoderOutput(vpx_codec_get_frame(Decoders[EDecoderType::Color].Handle, &Decoders[EDecoderType::Color].OutputIterator)
+				, Decoders[EDecoderType::Alpha].Handle ? vpx_codec_get_frame(Decoders[EDecoderType::Alpha].Handle, &Decoders[EDecoderType::Alpha].OutputIterator) : nullptr);
 		if (ConvResult == EConvertResult::Failure)
 		{
 			return IElectraDecoder::EOutputStatus::Error;
@@ -707,7 +769,7 @@ IElectraDecoder::EOutputStatus FVideoDecoderVPxElectra::HaveOutput()
 	}
 
 	// Decode any pending input first, even when we are to drain the decoder.
-	if (VideoDecoderHandle && !PendingDecoderInput.IsEmpty())
+	if (Decoders[EDecoderType::Color].Handle && !PendingDecoderInput.IsEmpty())
 	{
 		switch(DecodeNextPending())
 		{
@@ -733,11 +795,16 @@ IElectraDecoder::EOutputStatus FVideoDecoderVPxElectra::HaveOutput()
 			}
 		}
 	}
-	if (DecodeState == EDecodeState::Draining && VideoDecoderHandle)
+	if (DecodeState == EDecodeState::Draining && Decoders[EDecoderType::Color].Handle)
 	{
-		vpx_codec_err_t Result = vpx_codec_decode(VideoDecoderHandle, nullptr, 0, nullptr, 0);
+		vpx_codec_err_t Result = vpx_codec_decode(Decoders[EDecoderType::Color].Handle, nullptr, 0, nullptr, 0);
+		if (Decoders[EDecoderType::Alpha].Handle)
+		{
+			Result = vpx_codec_decode(Decoders[EDecoderType::Alpha].Handle, nullptr, 0, nullptr, 0);
+		}
 
-		EConvertResult ConvResult = ConvertDecoderOutput(vpx_codec_get_frame(VideoDecoderHandle, &VideoDecoderOutputIterator));
+		EConvertResult ConvResult = ConvertDecoderOutput(vpx_codec_get_frame(Decoders[EDecoderType::Color].Handle, &Decoders[EDecoderType::Color].OutputIterator)
+				, Decoders[EDecoderType::Alpha].Handle ? vpx_codec_get_frame(Decoders[EDecoderType::Alpha].Handle, &Decoders[EDecoderType::Alpha].OutputIterator) : nullptr);
 		if (ConvResult == EConvertResult::Failure)
 		{
 			return IElectraDecoder::EOutputStatus::Error;
@@ -766,7 +833,8 @@ TSharedPtr<IElectraDecoderOutput, ESPMode::ThreadSafe> FVideoDecoderVPxElectra::
 
 bool FVideoDecoderVPxElectra::InternalDecoderCreate(const TMap<FString, FVariant>& InAdditionalOptions)
 {
-	FMemory::Memzero(DecoderContext);
+	FMemory::Memzero(Decoders[EDecoderType::Color].Context);
+	FMemory::Memzero(Decoders[EDecoderType::Alpha].Context);
 	vpx_codec_err_t Result = VPX_CODEC_INCAPABLE;
 	switch(Codec4CC)
 	{
@@ -778,7 +846,7 @@ bool FVideoDecoderVPxElectra::InternalDecoderCreate(const TMap<FString, FVariant
 			vpx_codec_flags_t Flags = 0;
 			const int32 NumOfThreads = 1;
 			const vpx_codec_dec_cfg_t CodecConfig = { NumOfThreads, 0, 0 };
-			Result = vpx_codec_dec_init(&DecoderContext, vpx_codec_vp8_dx(), &CodecConfig, Flags);
+			Result = vpx_codec_dec_init(&Decoders[EDecoderType::Color].Context, vpx_codec_vp8_dx(), &CodecConfig, Flags);
 			break;
 		}
 		case Make4CC('v','p','0','9'):
@@ -790,7 +858,18 @@ bool FVideoDecoderVPxElectra::InternalDecoderCreate(const TMap<FString, FVariant
 			vpx_codec_flags_t Flags = 0;
 			const int32 NumOfThreads = 1;
 			const vpx_codec_dec_cfg_t CodecConfig = { NumOfThreads, 0, 0 };
-			Result = vpx_codec_dec_init(&DecoderContext, vpx_codec_vp9_dx(), &CodecConfig, Flags);
+			Result = vpx_codec_dec_init(&Decoders[EDecoderType::Color].Context, vpx_codec_vp9_dx(), &CodecConfig, Flags);
+
+			// If there is an additional alpha channel we need a 2nd decoder
+			if (InAdditionalOptions.Contains(DynamicSidebandData::VPxAlpha.ToString()) && Result == VPX_CODEC_OK)
+			{
+				Result = vpx_codec_dec_init(&Decoders[EDecoderType::Alpha].Context, vpx_codec_vp9_dx(), &CodecConfig, Flags);
+				if (Result == VPX_CODEC_OK)
+				{
+					Decoders[EDecoderType::Alpha].Handle = &Decoders[EDecoderType::Alpha].Context;
+					Decoders[EDecoderType::Alpha].OutputIterator = nullptr;
+				}
+			}
 			break;
 		}
 	}
@@ -798,29 +877,38 @@ bool FVideoDecoderVPxElectra::InternalDecoderCreate(const TMap<FString, FVariant
 	{
 		return PostError(Result, TEXT("Failed to create decoder"), ERRCODE_INTERNAL_COULD_NOT_CREATE_DECODER);
 	}
-	VideoDecoderHandle = &DecoderContext;
-	VideoDecoderOutputIterator = nullptr;
+	Decoders[EDecoderType::Color].Handle = &Decoders[EDecoderType::Color].Context;
+	Decoders[EDecoderType::Color].OutputIterator = nullptr;
 	return true;
 }
 
 void FVideoDecoderVPxElectra::InternalDecoderDestroy()
 {
-	if (VideoDecoderHandle)
+	if (Decoders[EDecoderType::Alpha].Handle)
+	{
+		/*vpx_codec_err_t Result =*/ vpx_codec_destroy(Decoders[EDecoderType::Alpha].Handle);
+		Decoders[EDecoderType::Alpha].Handle = nullptr;
+		FMemory::Memzero(Decoders[EDecoderType::Alpha].Context);
+		Decoders[EDecoderType::Alpha].OutputIterator = nullptr;
+	}
+	if (Decoders[EDecoderType::Color].Handle)
 	{
-		/*vpx_codec_err_t Result =*/ vpx_codec_destroy(VideoDecoderHandle);
-		VideoDecoderHandle = nullptr;
-		FMemory::Memzero(DecoderContext);
-		VideoDecoderOutputIterator = nullptr;
+		/*vpx_codec_err_t Result =*/ vpx_codec_destroy(Decoders[EDecoderType::Color].Handle);
+		Decoders[EDecoderType::Color].Handle = nullptr;
+		FMemory::Memzero(Decoders[EDecoderType::Color].Context);
+		Decoders[EDecoderType::Color].OutputIterator = nullptr;
 	}
 }
 
 
 
-FVideoDecoderVPxElectra::EConvertResult FVideoDecoderVPxElectra::ConvertDecoderOutput(vpx_image_t* InDecodedImage)
+FVideoDecoderVPxElectra::EConvertResult FVideoDecoderVPxElectra::ConvertDecoderOutput(vpx_image_t* InDecodedImage, vpx_image_t* InDecodedAlphaImage)
 {
 	if (!InDecodedImage)
 	{
-		VideoDecoderOutputIterator = nullptr;
+		check(InDecodedAlphaImage == nullptr);
+		Decoders[EDecoderType::Color].OutputIterator = nullptr;
+		Decoders[EDecoderType::Alpha].OutputIterator = nullptr;
 		return EConvertResult::GotEOS;
 	}
 
@@ -866,16 +954,32 @@ FVideoDecoderVPxElectra::EConvertResult FVideoDecoderVPxElectra::ConvertDecoderO
 
 	if (InDecodedImage->fmt == VPX_IMG_FMT_I420)
 	{
-		if (!ConvertDecodedImageToNV12orP010(NewOutput.Get(), InDecodedImage, PlatformDevice, PlatformDeviceVersion))
+		if (InDecodedAlphaImage == nullptr)
 		{
-			PostError(0, FString::Printf(TEXT("Failed to convert decoded image")), ERRCODE_INTERNAL_FAILED_TO_CONVERT_OUTPUT_SAMPLE);
-			return EConvertResult::Failure;
+			NewOutput->NumBuffers = 1;
+			NewOutput->ColorBufferFormat = EElectraDecoderPlatformPixelFormat::NV12;
+			NewOutput->ColorBufferEncoding = EElectraDecoderPlatformPixelEncoding::Native;
+			NewOutput->DecodedWidth = NewOutput->Width;
+			NewOutput->DecodedHeight = bUseGPUBuffers ? NewOutput->Height : (NewOutput->Height * 3 / 2);
+			if (!ConvertDecodedImageToNV12orP010(NewOutput.Get(), InDecodedImage, PlatformDevice, PlatformDeviceVersion))
+			{
+				PostError(0, FString::Printf(TEXT("Failed to convert decoded image")), ERRCODE_INTERNAL_FAILED_TO_CONVERT_OUTPUT_SAMPLE);
+				return EConvertResult::Failure;
+			}
+		}
+		else
+		{
+			NewOutput->NumBuffers = 1;
+			NewOutput->ColorBufferFormat = EElectraDecoderPlatformPixelFormat::B8G8R8A8;
+			NewOutput->ColorBufferEncoding = EElectraDecoderPlatformPixelEncoding::YCbCr_Alpha;
+			NewOutput->DecodedWidth = NewOutput->Width;
+			NewOutput->DecodedHeight = NewOutput->Height;
+			if (!ConvertDecodedImageWithAlpha(NewOutput.Get(), InDecodedImage, InDecodedAlphaImage, PlatformDevice, PlatformDeviceVersion))
+			{
+				PostError(0, FString::Printf(TEXT("Failed to convert decoded image")), ERRCODE_INTERNAL_FAILED_TO_CONVERT_OUTPUT_SAMPLE);
+				return EConvertResult::Failure;
+			}
 		}
-		NewOutput->NumBuffers = 1;
-		NewOutput->ColorBufferFormat = EElectraDecoderPlatformPixelFormat::NV12;
-		NewOutput->ColorBufferEncoding = EElectraDecoderPlatformPixelEncoding::Native;
-		NewOutput->DecodedWidth = NewOutput->Width;
-		NewOutput->DecodedHeight = bUseGPUBuffers ? NewOutput->Height : (NewOutput->Height * 3 / 2);
 	}
 	else if (InDecodedImage->fmt == VPX_IMG_FMT_I42016)
 	{
@@ -980,7 +1084,7 @@ bool FVideoDecoderVPxElectra::ConvertDecodedImageToNV12orP010(FVideoDecoderOutpu
 		uint32 BufferPitch;
 		D3D12ResourcePool->AllocateOutputDataAsBuffer(InNewOutput->GPUBuffer, BufferPitch);
 
-		// Correct pitch ton reflect resource's setup
+		// Correct pitch to reflect resource's setup
 		InNewOutput->ColorPitch = BufferPitch;
 
 		// Map the buffer so we can get a CPU address to the WC configured buffer
@@ -1052,3 +1156,138 @@ bool FVideoDecoderVPxElectra::ConvertDecodedImageToNV12orP010(FVideoDecoderOutpu
 
 	return true;
 }
+
+bool FVideoDecoderVPxElectra::ConvertDecodedImageWithAlpha(FVideoDecoderOutputVPxElectra* InNewOutput, const vpx_image_t* InDecodedImage, const vpx_image_t* InDecodedAlphaImage, void* PlatformDevice, int32 PlatformDeviceVersion) const
+{
+	// Can handle 8 bit 4:2:0 only at the moment.
+	// Would the alpha channel be 16 bit as well if the color channel is 16 bit?
+	if (InDecodedImage->fmt != VPX_IMG_FMT_I420)
+	{
+		return false;
+	}
+
+	const int32 w = InDecodedImage->d_w;
+	const int32 h = InDecodedImage->d_h;
+	const int32 aw = Align(w, 2);
+	const int32 ah = Align(h, 2);
+
+	const uint8* S0Y = (const uint8*)InDecodedImage->planes[0];
+	const uint8* S0A = (const uint8*)InDecodedAlphaImage->planes[0];
+	const uint8* SrcU = (const uint8*)InDecodedImage->planes[1];
+	const uint8* SrcV = (const uint8*)InDecodedImage->planes[2];
+	const int32 PitchY = InDecodedImage->stride[0];
+	const int32 PitchA = InDecodedAlphaImage->stride[0];
+	const int32 PitchU = InDecodedImage->stride[1];
+	const int32 PitchV = InDecodedImage->stride[2];
+	if (!S0Y || !SrcU || !SrcV || !S0A || PitchY != PitchA)
+	{
+		return false;
+	}
+
+
+	TSharedPtr<TArray<uint8>, ESPMode::ThreadSafe>& OutPixelBuffer = InNewOutput->ColorBuffer;
+	uint16* Dst0 = nullptr;
+	uint16* Dst1 = nullptr;
+	int32 DstPitch = 0;
+	int32 LineSkipDst = 0;
+
+#if ELECTRA_MEDIAGPUBUFFER_DX12
+	if (!PlatformDevice || PlatformDeviceVersion < 12000)
+#endif
+	{
+#if ELECTRA_MEDIAGPUBUFFER_DX12
+		InNewOutput->GPUBuffer.Resource = nullptr;
+		InNewOutput->GPUBuffer.Fence = nullptr;
+#endif
+
+		int32 AllocSize = aw * ah * sizeof(uint16) * 4;
+		OutPixelBuffer = MakeShared<TArray<uint8>, ESPMode::ThreadSafe>();
+		OutPixelBuffer->AddUninitialized(AllocSize);
+		Dst0 = (uint16*) OutPixelBuffer->GetData();
+		InNewOutput->ColorPitch = aw * sizeof(uint16) * 4;
+		DstPitch = aw * 4;
+		Dst1 = Dst0 + DstPitch;
+	}
+#if ELECTRA_MEDIAGPUBUFFER_DX12
+	else
+	{
+		TRefCountPtr D3D12Device(static_cast<ID3D12Device*>(PlatformDevice));
+
+		OutPixelBuffer.Reset();
+		InNewOutput->GPUBuffer.Resource = nullptr;
+
+		// Create the resource pool as needed...
+		if (!D3D12ResourcePool)
+		{
+			const uint32 BytesPerPixel = 8;
+			D3D12ResourcePool = MakeShared<FElectraMediaDecoderOutputBufferPool_DX12, ESPMode::ThreadSafe>(D3D12Device, MaxOutputBuffers, MaxWidth, MaxHeight, BytesPerPixel);
+		}
+
+		// Request resource and fence...
+		uint32 BufferPitch;
+		D3D12ResourcePool->AllocateOutputDataAsBuffer(InNewOutput->GPUBuffer, BufferPitch);
+
+		// Map the buffer so we can get a CPU address to the WC configured buffer
+		HRESULT Res = InNewOutput->GPUBuffer.Resource->Map(0, nullptr, (void**)&Dst0);
+		check(SUCCEEDED(Res));
+		InNewOutput->ColorPitch = BufferPitch;
+		DstPitch = BufferPitch / sizeof(uint16);
+		LineSkipDst = ((DstPitch / 4) - aw) * 4;
+		Dst1 = Dst0 + DstPitch;
+	}
+#endif
+	InNewOutput->NumBits = 16;
+	InNewOutput->ColorBufferFormat = EElectraDecoderPlatformPixelFormat::A16B16G16R16;
+	InNewOutput->ColorBufferEncoding = EElectraDecoderPlatformPixelEncoding::YCbCr_Alpha;
+
+	const int32 LineSkipY = PitchY - w;
+	const int32 LineSkipUV = PitchU - w/2;
+	const uint8* S1Y = S0Y + PitchY;
+	const uint8* S1A = S0A + PitchA;
+	#define EXPAND16(v) (((uint16)v << 8) | (uint16)v)
+	for(int32 y=0; y<h/2; ++y)
+	{
+		for(int32 x=0; x<w/2; ++x)
+		{
+			Dst0[0] = EXPAND16(S0A[0]);
+			Dst0[1] = EXPAND16(S0Y[0]);
+			Dst0[4] = EXPAND16(S0A[1]);
+			Dst0[5] = EXPAND16(S0Y[1]);
+			Dst1[0] = EXPAND16(S1A[0]);
+			Dst1[1] = EXPAND16(S1Y[0]);
+			Dst1[4] = EXPAND16(S1A[1]);
+			Dst1[5] = EXPAND16(S1Y[1]);
+			Dst0[2] = Dst0[6] = Dst1[2] = Dst1[6] = EXPAND16(*SrcU);
+			Dst0[3] = Dst0[7] = Dst1[3] = Dst1[7] = EXPAND16(*SrcV);
+			S0A += 2;
+			S0Y += 2;
+			S1A += 2;
+			S1Y += 2;
+			++SrcU;
+			++SrcV;
+			Dst0 += 8;
+			Dst1 += 8;
+		}
+		S0Y += LineSkipY + PitchY;
+		S1Y += LineSkipY + PitchY;
+		S0A += LineSkipY + PitchA;
+		S1A += LineSkipY + PitchA;
+		SrcU += LineSkipUV;
+		SrcV += LineSkipUV;
+		Dst0 += LineSkipDst + DstPitch;
+		Dst1 += LineSkipDst + DstPitch;
+	}
+
+#if ELECTRA_MEDIAGPUBUFFER_DX12
+	if (InNewOutput->GPUBuffer.Resource.IsValid())
+	{
+		InNewOutput->GPUBuffer.Resource->Unmap(0, nullptr);
+		// To be compatible with implementations that might do the copy into the resource async, we also signal a fence
+		// (strictly speaking we would not need to as this is all 100% synchronous and done before the GPU ever attempts to read from the resource)
+		InNewOutput->GPUBuffer.Fence->Signal(InNewOutput->GPUBuffer.FenceValue);
+	}
+#endif
+	return true;
+}
+
+
