Electra: Removal of the platform resource delegate and the wrapping plugin

diff --git a/Engine/Plugins/Media/VPxDecoderElectra/Source/VPxDecoderElectra/Private/VPxDecoder/ElectraMediaVPxDecoder.cpp b/Engine/Plugins/Media/VPxDecoderElectra/Source/VPxDecoderElectra/Private/VPxDecoder/ElectraMediaVPxDecoder.cpp
--- a/Engine/Plugins/Media/VPxDecoderElectra/Source/VPxDecoderElectra/Private/VPxDecoder/ElectraMediaVPxDecoder.cpp
+++ b/Engine/Plugins/Media/VPxDecoderElectra/Source/VPxDecoderElectra/Private/VPxDecoder/ElectraMediaVPxDecoder.cpp
@@ -12,7 +12,6 @@
 #include "IElectraDecoderFeaturesAndOptions.h"
 #include "IElectraDecoder.h"
 #include "IElectraDecoderOutputVideo.h"
-#include "IElectraDecoderResourceDelegate.h"
 #include "ElectraDecodersUtils.h"
 #include "Utils/Google/ElectraUtilsVPxVideo.h"
 #include "Utils/Google/ElectraBitstreamProcessorVPx.h"
@@ -23,6 +22,7 @@
 #define ENABLE_GPU_BUFFER_HELPER 1
 
 #if ENABLE_GPU_BUFFER_HELPER
+#include "ElectraDecodersPlatformResources.h"
 #include COMPILED_PLATFORM_HEADER(ElectraDecoderGPUBufferHelpers.h)
 #else
 #define ELECTRA_MEDIAGPUBUFFER_DX12 0
@@ -60,15 +60,6 @@ namespace DynamicSidebandData
 class FVideoDecoderVPxElectra;
 
 
-class FDecoderDefaultVideoOutputFormatVPxElectra : public IElectraDecoderDefaultVideoOutputFormat
-{
-public:
-	virtual ~FDecoderDefaultVideoOutputFormatVPxElectra()
-	{ }
-
-};
-
-
 class FVideoDecoderOutputVPxElectra : public IElectraDecoderVideoOutput, public IElectraDecoderVideoOutputImageBuffers
 {
 public:
@@ -112,8 +103,6 @@ public:
 		}
 		return nullptr;
 	}
-	IElectraDecoderVideoOutputTransferHandle* GetTransferHandle() const override
-	{ return nullptr; }
 	IElectraDecoderVideoOutput::EImageCopyResult CopyPlatformImage(IElectraDecoderVideoOutputCopyResources* InCopyResources) const override
 	{ return IElectraDecoderVideoOutput::EImageCopyResult::NotSupported; }
 
@@ -215,7 +204,7 @@ public:
 		OutOptions.Emplace(IElectraDecoderFeature::IsAdaptive, FVariant(false));
 	}
 
-	FVideoDecoderVPxElectra(const TMap<FString, FVariant>& InOptions, TSharedPtr<IElectraDecoderResourceDelegate, ESPMode::ThreadSafe> InResourceDelegate);
+	FVideoDecoderVPxElectra(const TMap<FString, FVariant>& InOptions);
 
 	virtual ~FVideoDecoderVPxElectra();
 
@@ -230,8 +219,6 @@ public:
 	IElectraDecoder::ECSDCompatibility IsCompatibleWith(const TMap<FString, FVariant>& CSDAndAdditionalOptions) override;
 	bool ResetToCleanStart() override;
 
-	TSharedPtr<IElectraDecoderDefaultOutputFormat, ESPMode::ThreadSafe> GetDefaultOutputFormatFromCSD(const TMap<FString, FVariant>& CSDAndAdditionalOptions) override;
-
 	EDecoderError DecodeAccessUnit(const FInputAccessUnit& InInputAccessUnit, const TMap<FString, FVariant>& InAdditionalOptions) override;
 	EDecoderError SendEndOfData() override;
 	EDecoderError Flush() override;
@@ -298,7 +285,6 @@ private:
 	bool ConvertDecodedImageWithAlpha(FVideoDecoderOutputVPxElectra* InNewOutput, const vpx_image_t* InDecodedImage, const vpx_image_t* InDecodedAlphaImage, void* PlatformDevice, int32 PlatformDeviceVersion) const;
 
 	TMap<FString, FVariant> InitialCreationOptions;
-	TWeakPtr<IElectraDecoderResourceDelegate, ESPMode::ThreadSafe> ResourceDelegate;
 
 	TQueue<TSharedPtr<FDecoderInput, ESPMode::NotThreadSafe>, EQueueMode::Spsc> PendingDecoderInput;
 	TArray<TSharedPtr<FDecoderInput, ESPMode::NotThreadSafe>> InDecoderInput;
@@ -378,9 +364,9 @@ public:
 		FVideoDecoderVPxElectra::GetConfigurationOptions(OutOptions);
 	}
 
-	TSharedPtr<IElectraDecoder, ESPMode::ThreadSafe> CreateDecoderForFormat(const FString& InCodecFormat, const TMap<FString, FVariant>& InOptions, TSharedPtr<IElectraDecoderResourceDelegate, ESPMode::ThreadSafe> InResourceDelegate) override
+	TSharedPtr<IElectraDecoder, ESPMode::ThreadSafe> CreateDecoderForFormat(const FString& InCodecFormat, const TMap<FString, FVariant>& InOptions) override
 	{
-		return MakeShared<FVideoDecoderVPxElectra, ESPMode::ThreadSafe>(InOptions, InResourceDelegate);
+		return MakeShared<FVideoDecoderVPxElectra, ESPMode::ThreadSafe>(InOptions);
 	}
 
 	static TSharedPtr<FVPxVideoDecoderElectraFactory, ESPMode::ThreadSafe> Self;
@@ -413,9 +399,8 @@ void FElectraMediaVPxDecoder::Shutdown()
 /*********************************************************************************************************************/
 /*********************************************************************************************************************/
 
-FVideoDecoderVPxElectra::FVideoDecoderVPxElectra(const TMap<FString, FVariant>& InOptions, TSharedPtr<IElectraDecoderResourceDelegate, ESPMode::ThreadSafe> InResourceDelegate)
+FVideoDecoderVPxElectra::FVideoDecoderVPxElectra(const TMap<FString, FVariant>& InOptions)
 	: InitialCreationOptions(InOptions)
-	, ResourceDelegate(InResourceDelegate)
 {
 	Codec4CC = (uint32)ElectraDecodersUtil::GetVariantValueSafeU64(InOptions, TEXT("codec_4cc"), 0);
 
@@ -481,11 +466,6 @@ bool FVideoDecoderVPxElectra::ResetToCleanStart()
 	return !LastError.IsSet();
 }
 
-TSharedPtr<IElectraDecoderDefaultOutputFormat, ESPMode::ThreadSafe> FVideoDecoderVPxElectra::GetDefaultOutputFormatFromCSD(const TMap<FString, FVariant>& CSDAndAdditionalOptions)
-{
-	return nullptr;
-}
-
 IElectraDecoder::EDecoderError FVideoDecoderVPxElectra::DecodeAccessUnit(const FInputAccessUnit& InInputAccessUnit, const TMap<FString, FVariant>& InAdditionalOptions)
 {
 	// If already in error do nothing!
@@ -945,11 +925,8 @@ FVideoDecoderVPxElectra::EConvertResult FVideoDecoderVPxElectra::ConvertDecoderO
 	int32 PlatformDeviceVersion = 0;
 	bool bUseGPUBuffers = false;
 #if ELECTRA_MEDIAGPUBUFFER_DX12
-	if (auto PinnedResourceDelegate = ResourceDelegate.Pin())
-	{
-		PinnedResourceDelegate->GetD3DDevice(&PlatformDevice, &PlatformDeviceVersion);
-		bUseGPUBuffers = (PlatformDevice && PlatformDeviceVersion >= 12000);
-	}
+	FElectraDecodersPlatformResources::GetD3DDeviceAndVersion(&PlatformDevice, &PlatformDeviceVersion);
+	bUseGPUBuffers = (PlatformDevice && PlatformDeviceVersion >= 12000);
 #endif
 
 	if (InDecodedImage->fmt == VPX_IMG_FMT_I420)
